Add a DefaultEngine.ini to UnrealVirtualizationTool with an example of how to enable an engine plugin with it.

#rnx

diff --git a/Engine/Programs/UnrealVirtualizationTool/Config/DefaultEngine.ini b/Engine/Programs/UnrealVirtualizationTool/Config/DefaultEngine.ini
new file mode 100644
--- /dev/null
+++ b/Engine/Programs/UnrealVirtualizationTool/Config/DefaultEngine.ini
@@ -0,0 +1,2 @@
+[Plugins]
+;+ProgramEnabledPlugins="AddYourPluginNameHere"
